Adding nesting to the Material Analyzer's collection creation
#rnx

diff --git a/Engine/Plugins/Editor/MaterialAnalyzer/Source/Private/SMaterialAnalyzer.cpp b/Engine/Plugins/Editor/MaterialAnalyzer/Source/Private/SMaterialAnalyzer.cpp
--- a/Engine/Plugins/Editor/MaterialAnalyzer/Source/Private/SMaterialAnalyzer.cpp
+++ b/Engine/Plugins/Editor/MaterialAnalyzer/Source/Private/SMaterialAnalyzer.cpp
@@ -391,6 +391,7 @@ TSharedRef< ITableRow > SMaterialAnalyzer::OnGenerateSuggestionRow(TSharedPtr<FP
 				[
 					SNew(SHorizontalBox)
 					+ SHorizontalBox::Slot()
+					.AutoWidth()
 					[
 						SNew(STextBlock)
 						.Visibility(this, &SMaterialAnalyzer::ShouldShowAdvancedRecommendations, Item)
@@ -463,6 +464,7 @@ FReply SMaterialAnalyzer::CreateLocalSuggestionCollection(TSharedPtr<FPermutatio
 	{
 		FCollectionManagerModule& CollectionManagerModule = FCollectionManagerModule::GetModule();
 
+		
 		FString FirstAssetString = CurrentlySelectedAsset.AssetName.ToString() + TEXT("_") + FString::FromInt(InSuggestion->Children.Num());
 		FName FirstAssetName = FName(*FirstAssetString);
 
@@ -477,7 +479,20 @@ FReply SMaterialAnalyzer::CreateLocalSuggestionCollection(TSharedPtr<FPermutatio
 			PackageNameSet.Add(FName(*FPaths::GetBaseFilename(*PackageToAdd, false)));
 		}
 
-		IAssetManagerEditorModule::Get().WriteCollection(FirstAssetName, ShareType, PackageNameSet, true);
+
+		ICollectionManager& CollectionManager = FCollectionManagerModule::GetModule().Get();
+		const FText MaterialAnalyzerText = LOCTEXT("MaterialAnalyzerPrefix", "MaterialAnalyzer");
+		FName ParentName = FName(*MaterialAnalyzerText.ToString());
+		if (!CollectionManager.CollectionExists(ParentName, ShareType))
+		{
+			CollectionManager.CreateCollection(ParentName, ShareType, ECollectionStorageMode::Static);
+		}
+
+		const bool bCollectionSucceeded = IAssetManagerEditorModule::Get().WriteCollection(FirstAssetName, ShareType, PackageNameSet, true);
+		if (bCollectionSucceeded)
+		{
+			CollectionManager.ReparentCollection(FirstAssetName, ShareType, ParentName, ShareType);
+		}
 	}
 	return FReply::Handled();
 }
