Adding recursive expansion to the material analyzer tree
#rnx

diff --git a/Engine/Plugins/Editor/MaterialAnalyzer/Source/Private/SMaterialAnalyzer.cpp b/Engine/Plugins/Editor/MaterialAnalyzer/Source/Private/SMaterialAnalyzer.cpp
--- a/Engine/Plugins/Editor/MaterialAnalyzer/Source/Private/SMaterialAnalyzer.cpp
+++ b/Engine/Plugins/Editor/MaterialAnalyzer/Source/Private/SMaterialAnalyzer.cpp
@@ -136,6 +136,7 @@ void SMaterialAnalyzer::Construct(const FArguments& InArgs, const TSharedRef<SDo
 						.TreeItemsSource(&MaterialTreeRoot)
 						.OnGenerateRow(this, &SMaterialAnalyzer::HandleReflectorTreeGenerateRow)
 						.OnGetChildren(this, &SMaterialAnalyzer::HandleReflectorTreeGetChildren)
+						.OnSetExpansionRecursive(this, &SMaterialAnalyzer::HandleReflectorTreeRecursiveExpansion)
 						.HeaderRow
 						(
 							SNew(SHeaderRow)
@@ -590,6 +591,23 @@ void SMaterialAnalyzer::HandleReflectorTreeGetChildren(FAnalyzedMaterialNodeRef
 	OutChildren = InMaterialNode->GetChildNodes();
 }
 
+void SMaterialAnalyzer::HandleReflectorTreeRecursiveExpansion(FAnalyzedMaterialNodeRef InTreeNode, bool bIsItemExpanded)
+{
+	TArray< FAnalyzedMaterialNodeRef > Children = InTreeNode->GetChildNodes();
+
+	if (Children.Num())
+	{
+		MaterialTree->SetItemExpansion(InTreeNode, bIsItemExpanded);
+		const bool bShouldSaveState = true;
+
+		for (int32 ChildIndex = 0; ChildIndex < Children.Num(); ++ChildIndex)
+		{
+			FAnalyzedMaterialNodeRef Child = Children[ChildIndex];
+			MaterialTree->SetItemExpansion(Child, bIsItemExpanded);
+		}
+	}
+}
+
 FAnalyzedMaterialNodePtr FBuildBasicMaterialTreeAsyncTask::FindOrMakeBranchNode(FAnalyzedMaterialNodePtr ParentNode, const FAssetData* ChildData)
 {
 	check(ChildData != nullptr);
diff --git a/Engine/Plugins/Editor/MaterialAnalyzer/Source/Private/SMaterialAnalyzer.h b/Engine/Plugins/Editor/MaterialAnalyzer/Source/Private/SMaterialAnalyzer.h
--- a/Engine/Plugins/Editor/MaterialAnalyzer/Source/Private/SMaterialAnalyzer.h
+++ b/Engine/Plugins/Editor/MaterialAnalyzer/Source/Private/SMaterialAnalyzer.h
@@ -273,6 +273,14 @@ private:
 	/** Callback for getting the child items of the given reflector tree node. */
 	void HandleReflectorTreeGetChildren(FAnalyzedMaterialNodeRef InReflectorNode, TArray<FAnalyzedMaterialNodeRef>& OutChildren);
 
+	/**
+	 * Called to recursively expand/collapse the children of the given item
+	 *
+	 * @param InTreeNode		The node that was expanded or collapsed
+	 * @param bIsItemExpanded	True if the item is expanded, false if it is collapsed
+	 */
+	void HandleReflectorTreeRecursiveExpansion(FAnalyzedMaterialNodeRef InTreeNode, bool bIsItemExpanded);
+
 	bool IsMaterialSelectionAllowed() const
 	{
 		return bAllowMaterialSelection;
