Adding "Browse To Asset in Content Browser" button and alternate row styling for the Material Analyzer
#rnx

diff --git a/Engine/Plugins/Editor/MaterialAnalyzer/Source/Private/AnalyzedMaterialNode.h b/Engine/Plugins/Editor/MaterialAnalyzer/Source/Private/AnalyzedMaterialNode.h
--- a/Engine/Plugins/Editor/MaterialAnalyzer/Source/Private/AnalyzedMaterialNode.h
+++ b/Engine/Plugins/Editor/MaterialAnalyzer/Source/Private/AnalyzedMaterialNode.h
@@ -161,6 +161,7 @@ public:
 	FString Path;
 	FName ObjectPath;
 	FAnalyzedMaterialNodePtr Parent;
+	FAssetData AssetData;
 
 	TArray<FBasePropertyOverrideNodeRef> BasePropertyOverrides;
 	TArray<FStaticMaterialLayerParameterNodeRef> MaterialLayerParameters;
diff --git a/Engine/Plugins/Editor/MaterialAnalyzer/Source/Private/SAnalyzedMaterialNodeWidgetItem.cpp b/Engine/Plugins/Editor/MaterialAnalyzer/Source/Private/SAnalyzedMaterialNodeWidgetItem.cpp
--- a/Engine/Plugins/Editor/MaterialAnalyzer/Source/Private/SAnalyzedMaterialNodeWidgetItem.cpp
+++ b/Engine/Plugins/Editor/MaterialAnalyzer/Source/Private/SAnalyzedMaterialNodeWidgetItem.cpp
@@ -6,6 +6,9 @@
 #include "Widgets/Images/SImage.h"
 #include "Engine/EngineTypes.h"
 #include "MaterialShaderType.h"
+#include "PropertyCustomizationHelpers.h"
+#include "Editor/EditorEngine.h"
+#include "EditorStyleSet.h"
 
 #define LOCTEXT_NAMESPACE "MaterialAnalyzer"
 
@@ -18,20 +21,22 @@ FName SAnalyzedMaterialNodeWidgetItem::NAME_StaticComponentMaskParameters(TEXT("
 
 void SAnalyzedMaterialNodeWidgetItem::Construct(const FArguments& InArgs, const TSharedRef<STableViewBase>& InOwnerTableView)
 {
-	this->MaterialInfo = InArgs._MaterialInfoToVisualize;
-
-	this->SetPadding(0);
-
-	this->BasePropertyOverrideNodes = MaterialInfo->BasePropertyOverrides;
-	this->StaticSwitchNodes = MaterialInfo->StaticSwitchParameters;
-	this->StaticMaterialLayerNodes = MaterialInfo->MaterialLayerParameters;
-	this->StaticComponentMaskNodes = MaterialInfo->StaticComponentMaskParameters;
-
+	MaterialInfo = InArgs._MaterialInfoToVisualize;
+
+	SetPadding(0);
+	BasePropertyOverrideNodes = MaterialInfo->BasePropertyOverrides;
+	StaticSwitchNodes = MaterialInfo->StaticSwitchParameters;
+	StaticMaterialLayerNodes = MaterialInfo->MaterialLayerParameters;
+	StaticComponentMaskNodes = MaterialInfo->StaticComponentMaskParameters;
+	AssetData = MaterialInfo->AssetData;
 	CachedMaterialName = FText::FromString(MaterialInfo->Path);
 	TotalNumberOfChildren = MaterialInfo->TotalNumberOfChildren();
 	NumberOfChildren = MaterialInfo->ActualNumberOfChildren();
 
-	SMultiColumnTableRow< FAnalyzedMaterialNodeRef >::Construct(SMultiColumnTableRow< FAnalyzedMaterialNodeRef >::FArguments().Padding(0), InOwnerTableView);
+	SMultiColumnTableRow< FAnalyzedMaterialNodeRef >::Construct(SMultiColumnTableRow< FAnalyzedMaterialNodeRef >::FArguments()
+		.Padding(0)
+		.Style(FEditorStyle::Get(), "DataTableEditor.CellListViewRow")
+		, InOwnerTableView);
 }
 
 TSharedRef<SWidget> SAnalyzedMaterialNodeWidgetItem::GenerateWidgetForColumn(const FName& ColumnName)
@@ -39,10 +44,9 @@ TSharedRef<SWidget> SAnalyzedMaterialNodeWidgetItem::GenerateWidgetForColumn(con
 	if(ColumnName == NAME_MaterialName)
 	{
 		return SNew(SHorizontalBox)
-
 		+ SHorizontalBox::Slot()
 		.AutoWidth()
-		.VAlign(VAlign_Top)
+		.VAlign(VAlign_Center)
 		[
 			SNew(SExpanderArrow, SharedThis(this))
 			.IndentAmount(16)
@@ -50,10 +54,25 @@ TSharedRef<SWidget> SAnalyzedMaterialNodeWidgetItem::GenerateWidgetForColumn(con
 		+SHorizontalBox::Slot()
 		.AutoWidth()
 		.Padding(2.0f, 0.0f)
-		.VAlign(VAlign_Top)
+		.VAlign(VAlign_Center)
 		[
 			SNew(STextBlock)
 			.Text(this, &SAnalyzedMaterialNodeWidgetItem::GetMaterialName)
+		]
+		+ SHorizontalBox::Slot()
+		.Padding(1.0f, 1.0f)
+		.AutoWidth()
+		[
+			SNew(SButton)
+			.OnClicked(this, &SAnalyzedMaterialNodeWidgetItem::FindInContentBrowser)
+			.ButtonStyle(FEditorStyle::Get(), "HoverHintOnly")
+			.ToolTipText(LOCTEXT("BrowseInContentBrowser", "Browse to this asset in the Content Browser"))
+			.DesiredSizeScale(FVector2D(0.4f, 0.4f))
+			.VAlign(VAlign_Center)
+			[
+				SNew(SImage)
+				.Image(FSlateIcon(FEditorStyle::GetStyleSetName(), "SystemWideCommands.FindInContentBrowser").GetIcon())
+			]
 		];
 	}
 	else if(ColumnName == NAME_NumberOfChildren)
@@ -86,6 +105,14 @@ TSharedRef<SWidget> SAnalyzedMaterialNodeWidgetItem::GenerateWidgetForColumn(con
 	return SNullWidget::NullWidget;
 }
 
+FReply SAnalyzedMaterialNodeWidgetItem::FindInContentBrowser()
+{
+	TArray<FAssetData> AssetDataArray;
+	AssetDataArray.Add(AssetData);
+	GEditor->SyncBrowserToObjects(AssetDataArray);
+	return FReply::Handled();
+}
+
 
 
 template <typename NodeType>
diff --git a/Engine/Plugins/Editor/MaterialAnalyzer/Source/Private/SAnalyzedMaterialNodeWidgetItem.h b/Engine/Plugins/Editor/MaterialAnalyzer/Source/Private/SAnalyzedMaterialNodeWidgetItem.h
--- a/Engine/Plugins/Editor/MaterialAnalyzer/Source/Private/SAnalyzedMaterialNodeWidgetItem.h
+++ b/Engine/Plugins/Editor/MaterialAnalyzer/Source/Private/SAnalyzedMaterialNodeWidgetItem.h
@@ -45,6 +45,7 @@ public:
 	// SMultiColumnTableRow overrides
 	virtual TSharedRef<SWidget> GenerateWidgetForColumn(const FName& ColumnName) override;
 
+	FReply FindInContentBrowser();
 	FText GetMaterialName() const
 	{
 		return CachedMaterialName;
@@ -62,7 +63,7 @@ protected:
 	FText CachedMaterialName;
 	int TotalNumberOfChildren;
 	int NumberOfChildren;
-
+	FAssetData AssetData;
 	TArray<FBasePropertyOverrideNodeRef> BasePropertyOverrideNodes;
 	TArray<FStaticSwitchParameterNodeRef> StaticSwitchNodes;
 	TArray<FStaticComponentMaskParameterNodeRef> StaticComponentMaskNodes;
diff --git a/Engine/Plugins/Editor/MaterialAnalyzer/Source/Private/SMaterialAnalyzer.cpp b/Engine/Plugins/Editor/MaterialAnalyzer/Source/Private/SMaterialAnalyzer.cpp
--- a/Engine/Plugins/Editor/MaterialAnalyzer/Source/Private/SMaterialAnalyzer.cpp
+++ b/Engine/Plugins/Editor/MaterialAnalyzer/Source/Private/SMaterialAnalyzer.cpp
@@ -625,6 +625,7 @@ FAnalyzedMaterialNodePtr FBuildBasicMaterialTreeAsyncTask::FindOrMakeBranchNode(
 		ChildNode.Path = ChildData->AssetName.ToString();
 		ChildNode.ObjectPath = ChildData->ObjectPath;
 		ChildNode.Parent = ParentNode;
+		ChildNode.AssetData = *ChildData;
 		NodesToSearch.Add(FAnalyzedMaterialNodeRef(new FAnalyzedMaterialNode(ChildNode)));
 		OutNode = &NodesToSearch[NodesToSearch.Num() - 1];
 	}
diff --git a/Engine/Source/Editor/EditorStyle/Private/SlateEditorStyle.cpp b/Engine/Source/Editor/EditorStyle/Private/SlateEditorStyle.cpp
--- a/Engine/Source/Editor/EditorStyle/Private/SlateEditorStyle.cpp
+++ b/Engine/Source/Editor/EditorStyle/Private/SlateEditorStyle.cpp
@@ -6931,6 +6931,20 @@ void FSlateEditorStyle::FStyle::SetupToolkitStyles()
 
 		Set("MaterialEditor.AddSubstanceSpecialToolBar", new IMAGE_BRUSH("Icons/icon_SubstanceFactory_White_40x", Icon40x40));
 		Set("MaterialEditor.AddSubstanceSpecialMenu", new IMAGE_BRUSH("Icons/icon_SubstanceFactory_White_40x", Icon32x32));
+
+
+		Set("MaterialEditor.CellListViewRow", FTableRowStyle(NormalTableRowStyle)
+			.SetEvenRowBackgroundBrush(IMAGE_BRUSH("PropertyView/DetailCategoryMiddle", FVector2D(16, 16), FLinearColor(0.5f, 0.5f, 0.5f)))
+			.SetEvenRowBackgroundHoveredBrush(IMAGE_BRUSH("PropertyView/DetailCategoryMiddle_Hovered", FVector2D(16, 16), FLinearColor(0.5f, 0.5f, 0.5f)))
+			.SetOddRowBackgroundBrush(IMAGE_BRUSH("PropertyView/DetailCategoryMiddle", FVector2D(16, 16), FLinearColor(0.35f, 0.35f, 0.35f)))
+			.SetOddRowBackgroundHoveredBrush(IMAGE_BRUSH("PropertyView/DetailCategoryMiddle_Hovered", FVector2D(16, 16), FLinearColor(0.35f, 0.35f, 0.35f)))
+			.SetActiveBrush(IMAGE_BRUSH("Common/Selection", Icon8x8, FLinearColor(0.075f, 0.075f, 0.075f)))
+			.SetActiveHoveredBrush(IMAGE_BRUSH("Common/Selection", Icon8x8, FLinearColor(0.075f, 0.075f, 0.075f)))
+			.SetInactiveBrush(IMAGE_BRUSH("Common/Selection", Icon8x8, FLinearColor(0.075f, 0.075f, 0.075f)))
+			.SetInactiveHoveredBrush(IMAGE_BRUSH("Common/Selection", Icon8x8, FLinearColor(0.075f, 0.075f, 0.075f)))
+			.SetTextColor(DefaultForeground)
+			.SetSelectedTextColor(DefaultForeground)
+		);
 	}
 
 	// Material Instance Editor
@@ -6943,8 +6957,10 @@ void FSlateEditorStyle::FStyle::SetupToolkitStyles()
 		Set("MaterialInstanceEditor.StackBodyBlend", new FSlateColorBrush(FLinearColor(0.06f, 0.06f, 0.06f, 1.0f)));
 		Set("MaterialInstanceEditor.StackBodyDragged", new FSlateColorBrush(FLinearColor(0.0f, 0.0f, 0.0f, 1.0f)));
 		Set("MaterialInstanceEditor.StackBody_Highlighted", new BOX_BRUSH("Common/TextBox_Special_Active", FMargin(8.0f / 32.0f)));
+
 	}
 
+
 	// Sound Class Editor
 	{
 		Set( "SoundClassEditor.Tabs.Properties", new IMAGE_BRUSH( "/Icons/icon_tab_SelectionDetails_16x", Icon16x16 ) );
