Fixing non-unity compile errors
#rnx

diff --git a/Engine/Plugins/Editor/MaterialAnalyzer/Source/Private/AnalyzedMaterialNode.h b/Engine/Plugins/Editor/MaterialAnalyzer/Source/Private/AnalyzedMaterialNode.h
--- a/Engine/Plugins/Editor/MaterialAnalyzer/Source/Private/AnalyzedMaterialNode.h
+++ b/Engine/Plugins/Editor/MaterialAnalyzer/Source/Private/AnalyzedMaterialNode.h
@@ -3,6 +3,7 @@
 #pragma once
 
 #include "CoreMinimal.h"
+#include "AssetData.h"
 
 struct FBasePropertyOverrideNode
 {
diff --git a/Engine/Plugins/Editor/MaterialAnalyzer/Source/Private/SAnalyzedMaterialNodeWidgetItem.cpp b/Engine/Plugins/Editor/MaterialAnalyzer/Source/Private/SAnalyzedMaterialNodeWidgetItem.cpp
--- a/Engine/Plugins/Editor/MaterialAnalyzer/Source/Private/SAnalyzedMaterialNodeWidgetItem.cpp
+++ b/Engine/Plugins/Editor/MaterialAnalyzer/Source/Private/SAnalyzedMaterialNodeWidgetItem.cpp
@@ -7,7 +7,7 @@
 #include "Engine/EngineTypes.h"
 #include "MaterialShaderType.h"
 #include "PropertyCustomizationHelpers.h"
-#include "Editor/EditorEngine.h"
+#include "Editor.h"
 #include "EditorStyleSet.h"
 
 #define LOCTEXT_NAMESPACE "MaterialAnalyzer"
diff --git a/Engine/Plugins/Editor/MaterialAnalyzer/Source/Private/SAnalyzedMaterialNodeWidgetItem.h b/Engine/Plugins/Editor/MaterialAnalyzer/Source/Private/SAnalyzedMaterialNodeWidgetItem.h
--- a/Engine/Plugins/Editor/MaterialAnalyzer/Source/Private/SAnalyzedMaterialNodeWidgetItem.h
+++ b/Engine/Plugins/Editor/MaterialAnalyzer/Source/Private/SAnalyzedMaterialNodeWidgetItem.h
@@ -11,6 +11,7 @@
 #include "Widgets/Views/SListView.h"
 #include "AnalyzedMaterialNode.h"
 #include "Widgets/Input/SButton.h"
+#include "AssetData.h"
 
 class SAnalyzedMaterialNodeWidgetItem
 	: public SMultiColumnTableRow<FAnalyzedMaterialNodeRef>
