Fixed localization gather warnings

diff --git a/Engine/Plugins/Compositing/Composure/Source/ComposureLayersEditor/Private/CompElementCollectionViewModel.cpp b/Engine/Plugins/Compositing/Composure/Source/ComposureLayersEditor/Private/CompElementCollectionViewModel.cpp
--- a/Engine/Plugins/Compositing/Composure/Source/ComposureLayersEditor/Private/CompElementCollectionViewModel.cpp
+++ b/Engine/Plugins/Compositing/Composure/Source/ComposureLayersEditor/Private/CompElementCollectionViewModel.cpp
@@ -846,7 +846,7 @@ void FCompElementCollectionViewModel::CreateTopLevelElement_Executed()
 		HighlightedClasses.Add(ACompositingElement::StaticClass());
 	}
 
-	UClass* ChosenClass = CompElementCollectionViewModel_Impl::PromptForElementClass(LOCTEXT("PickElementClass", "Pick an Comp Class"), HighlightedClasses);
+	UClass* ChosenClass = CompElementCollectionViewModel_Impl::PromptForElementClass(LOCTEXT("PickCompClass", "Pick an Comp Class"), HighlightedClasses);
 
 	if (ChosenClass)
 	{
diff --git a/Engine/Plugins/Editor/MaterialAnalyzer/Source/Private/SMaterialAnalyzer.cpp b/Engine/Plugins/Editor/MaterialAnalyzer/Source/Private/SMaterialAnalyzer.cpp
--- a/Engine/Plugins/Editor/MaterialAnalyzer/Source/Private/SMaterialAnalyzer.cpp
+++ b/Engine/Plugins/Editor/MaterialAnalyzer/Source/Private/SMaterialAnalyzer.cpp
@@ -107,7 +107,7 @@ void SMaterialAnalyzer::Construct(const FArguments& InArgs, const TSharedRef<SDo
 				.VAlign(VAlign_Center)
 				[
 					SNew(STextBlock)
-					.Text(LOCTEXT("MaterialToAnalyze", "Material To Analyze: "))
+					.Text(LOCTEXT("MaterialToAnalyzeLabel", "Material To Analyze: "))
 				]
 				+ SHorizontalBox::Slot()
 				.FillWidth(0.5f)
@@ -140,18 +140,18 @@ void SMaterialAnalyzer::Construct(const FArguments& InArgs, const TSharedRef<SDo
 						(
 							SNew(SHeaderRow)
 							+ SHeaderRow::Column(SAnalyzedMaterialNodeWidgetItem::NAME_MaterialName)
-							.DefaultLabel(LOCTEXT("MaterialName", "Material Name"))
+							.DefaultLabel(LOCTEXT("MaterialNameLabel", "Material Name"))
 							.FillWidth(0.80f)
 							+ SHeaderRow::Column(SAnalyzedMaterialNodeWidgetItem::NAME_NumberOfChildren)
-							.DefaultLabel(LOCTEXT("NumberOfMaterialChildren", "Number of Children (Direct/Total)"))
+							.DefaultLabel(LOCTEXT("NumberOfMaterialChildrenLabel", "Number of Children (Direct/Total)"))
 							+ SHeaderRow::Column(SAnalyzedMaterialNodeWidgetItem::NAME_BasePropertyOverrides)
-							.DefaultLabel(LOCTEXT("BasePropertyOverrides", "Base Property Overrides"))
+							.DefaultLabel(LOCTEXT("BasePropertyOverridesLabel", "Base Property Overrides"))
 							+ SHeaderRow::Column(SAnalyzedMaterialNodeWidgetItem::NAME_MaterialLayerParameters)
-							.DefaultLabel(LOCTEXT("MaterialLayerParameters", "Material Layer Parameters"))
+							.DefaultLabel(LOCTEXT("MaterialLayerParametersLabel", "Material Layer Parameters"))
 							+ SHeaderRow::Column(SAnalyzedMaterialNodeWidgetItem::NAME_StaticSwitchParameters)
-							.DefaultLabel(LOCTEXT("StaticSwitchParameters", "Static Switch Parameters"))
+							.DefaultLabel(LOCTEXT("StaticSwitchParametersLabel", "Static Switch Parameters"))
 							+ SHeaderRow::Column(SAnalyzedMaterialNodeWidgetItem::NAME_StaticComponentMaskParameters)
-							.DefaultLabel(LOCTEXT("StaticComponenetMaskParameters", "Static Component Mask Parameters"))
+							.DefaultLabel(LOCTEXT("StaticComponenetMaskParametersLabel", "Static Component Mask Parameters"))
 						)
 					]
 				]
@@ -167,7 +167,7 @@ void SMaterialAnalyzer::Construct(const FArguments& InArgs, const TSharedRef<SDo
 						.VAlign(VAlign_Center)
 						[
 							SNew(STextBlock)
-							.Text(LOCTEXT("Suggestions", "Suggestions"))
+							.Text(LOCTEXT("SuggestionsLabel", "Suggestions"))
 						]
 						+ SVerticalBox::Slot()
 						[
@@ -185,7 +185,7 @@ void SMaterialAnalyzer::Construct(const FArguments& InArgs, const TSharedRef<SDo
 				.VAlign(VAlign_Center)
 				[
 					SAssignNew(StatusBox, STextBlock)
-					.Text(LOCTEXT("Done", "Done"))
+					.Text(LOCTEXT("DoneLabel", "Done"))
 				]
 				+SHorizontalBox::Slot()
 				.HAlign(HAlign_Right)
@@ -283,7 +283,7 @@ void SMaterialAnalyzer::Tick(const FGeometry& AllottedGeometry, const double InC
 		{
 			delete BuildBaseMaterialTreeTask;
 			BuildBaseMaterialTreeTask = nullptr;
-			AsyncWorkFinished(FText::Format(FTextFormat(LOCTEXT("MaterialAnalyzer_DoneWithMaterialInterfaces", "Done with {0} MaterialInterfaces")), GetTotalNumberOfMaterialNodes()));
+			AsyncWorkFinished(FText::Format(FTextFormat(LOCTEXT("DoneWithMaterialInterfaces", "Done with {0} MaterialInterfaces")), GetTotalNumberOfMaterialNodes()));
 
 		}
 
@@ -546,7 +546,7 @@ void SMaterialAnalyzer::BuildBasicMaterialTree()
 		BuildBaseMaterialTreeTask = new FAsyncTask<FBuildBasicMaterialTreeAsyncTask>(AllMaterialTreeRoots, AssetDataArray);
 		BuildBaseMaterialTreeTask->StartBackgroundTask();
 
-		StartAsyncWork(LOCTEXT("BuildingBasicTree", "Building basic material tree"));
+		StartAsyncWork(LOCTEXT("BuildingBasicTree", "Building Basic MaterialTree"));
 
 		if (StatusThrobber.IsValid())
 		{
diff --git a/Engine/Source/Editor/BlueprintGraph/Private/K2Node_GetSubsystem.cpp b/Engine/Source/Editor/BlueprintGraph/Private/K2Node_GetSubsystem.cpp
--- a/Engine/Source/Editor/BlueprintGraph/Private/K2Node_GetSubsystem.cpp
+++ b/Engine/Source/Editor/BlueprintGraph/Private/K2Node_GetSubsystem.cpp
@@ -85,7 +85,7 @@ FText UK2Node_GetSubsystem::GetNodeTitle(ENodeTitleType::Type TitleType) const
 		}
 
 		const FString& ClassName = CustomClass->GetName();
-		return FText::FormatNamed(NSLOCTEXT("K2Node", "GetSubsystemNodeTitleFormat", "Get {ClassName}"), TEXT("ClassName"), FText::FromString(ClassName));
+		return FText::FormatNamed(NSLOCTEXT("K2Node", "GetSubsystem_NodeTitleFormat", "Get {ClassName}"), TEXT("ClassName"), FText::FromString(ClassName));
 	}
 
 	return GetTooltipText();
@@ -229,14 +229,14 @@ FText UK2Node_GetSubsystem::GetMenuCategory() const
 {
 	if (CustomClass->IsChildOf<UGameInstanceSubsystem>())
 	{
-		return NSLOCTEXT("K2Node", "GameInstanceSubsystems", "GameInstance Subsystems");
+		return NSLOCTEXT("K2Node", "GetSubsystem_GameInstanceSubsystemsMenuCategory", "GameInstance Subsystems");
 	}
 	else if (CustomClass->IsChildOf<ULocalPlayerSubsystem>())
 	{
-		return NSLOCTEXT("K2Node", "LocalPlayerSubsystems", "LocalPlayer Subsystems");
+		return NSLOCTEXT("K2Node", "GetSubsystem_LocalPlayerSubsystemsMenuCategory", "LocalPlayer Subsystems");
 	}
 
-	return NSLOCTEXT("K2Node", "InvalidSubsystemType", "Invalid Subsystem Type");
+	return NSLOCTEXT("K2Node", "GetSubsystem_InvalidSubsystemTypeMenuCategory", "Invalid Subsystem Type");
 }
 
 FText UK2Node_GetSubsystem::GetTooltipText() const
@@ -246,16 +246,16 @@ FText UK2Node_GetSubsystem::GetTooltipText() const
 		FText SubsystemTypeText;
 		if (CustomClass->IsChildOf<UGameInstanceSubsystem>())
 		{
-			SubsystemTypeText =  NSLOCTEXT("K2Node", "GameInstanceSubsystem", "GameInstance Subsystem");
+			SubsystemTypeText =  NSLOCTEXT("K2Node", "GetSubsystem_GameInstanceSubsystemTooltip", "GameInstance Subsystem");
 		}
 		else
 		{
-			SubsystemTypeText = NSLOCTEXT("K2Node", "LocalPlayerSubsystem", "LocalPlayer Subsystem");
+			SubsystemTypeText = NSLOCTEXT("K2Node", "GetSubsystem_LocalPlayerSubsystemTooltip", "LocalPlayer Subsystem");
 		}
-		return FText::FormatNamed(NSLOCTEXT("K2Node", "GetSubsystemTooltipFormat", "Get {ClassName} a {SubsystemType}"), TEXT("ClassName"), CustomClass->GetDisplayNameText(), TEXT("SubsystemType"), SubsystemTypeText);
+		return FText::FormatNamed(NSLOCTEXT("K2Node", "GetSubsystem_TooltipFormat", "Get {ClassName} a {SubsystemType}"), TEXT("ClassName"), CustomClass->GetDisplayNameText(), TEXT("SubsystemType"), SubsystemTypeText);
 	}
 
-	return NSLOCTEXT("K2Node", "InvalidSubsystemType", "InvalidSubsystemType");
+	return NSLOCTEXT("K2Node", "GetSubsystem_InvalidSubsystemTypeTooltip", "Invalid Subsystem Type");
 }
 
 UEdGraphPin* UK2Node_GetSubsystem::GetWorldContextPin() const
@@ -412,17 +412,17 @@ void UK2Node_GetSubsystemFromPC::GetMenuActions(FBlueprintActionDatabaseRegistra
 
 FText UK2Node_GetSubsystemFromPC::GetMenuCategory() const
 {
-	return NSLOCTEXT("K2Node", "LocalPlayerSubsystemsFromPC", "PlayerController|LocalPlayer Subsystems");
+	return NSLOCTEXT("K2Node", "GetSubsystemFromPC_MenuCategory", "PlayerController|LocalPlayer Subsystems");
 }
 
 FText UK2Node_GetSubsystemFromPC::GetTooltipText() const
 {
 	if (CustomClass)
 	{
-		return FText::FormatNamed(NSLOCTEXT("K2Node", "GetSubsystemTooltipFormat", "Get {ClassName} from Player Controller"), TEXT("ClassName"), CustomClass->GetDisplayNameText());
+		return FText::FormatNamed(NSLOCTEXT("K2Node", "GetSubsystemFromPC_TooltipFormat", "Get {ClassName} from Player Controller"), TEXT("ClassName"), CustomClass->GetDisplayNameText());
 	}
 
-	return NSLOCTEXT("K2Node", "InvalidSubsystemType", "InvalidSubsystemType");
+	return NSLOCTEXT("K2Node", "GetSubsystemFromPC_InvalidSubsystemTypeTooltip", "Invalid Subsystem Type");
 }
 
 UEdGraphPin* UK2Node_GetSubsystemFromPC::GetPlayerControllerPin() const
diff --git a/Engine/Source/Editor/Blutility/Private/AssetTypeActions_EditorUtilityWidgetBlueprint.cpp b/Engine/Source/Editor/Blutility/Private/AssetTypeActions_EditorUtilityWidgetBlueprint.cpp
--- a/Engine/Source/Editor/Blutility/Private/AssetTypeActions_EditorUtilityWidgetBlueprint.cpp
+++ b/Engine/Source/Editor/Blutility/Private/AssetTypeActions_EditorUtilityWidgetBlueprint.cpp
@@ -49,8 +49,8 @@ void FAssetTypeActions_EditorUtilityWidgetBlueprint::GetActions(const TArray<UOb
 	auto Blueprints = GetTypedWeakObjectPtrs<UWidgetBlueprint>(InObjects);
 
 	MenuBuilder.AddMenuEntry(
-		LOCTEXT("Blutility_Edit", "Edit Blueprint"),
-		LOCTEXT("Blutility_EditTooltip", "Opens the selected blueprints in the full blueprint editor."),
+		LOCTEXT("EditorUtilityWidget_Edit", "Edit Blueprint"),
+		LOCTEXT("EditorUtilityWidget_EditTooltip", "Opens the selected blueprints in the full blueprint editor."),
 		FSlateIcon(),
 		FUIAction(
 			FExecuteAction::CreateSP(this, &FAssetTypeActions_EditorUtilityWidgetBlueprint::ExecuteEdit, Blueprints),
@@ -124,7 +124,7 @@ void FAssetTypeActions_EditorUtilityWidgetBlueprint::ExecuteEdit(FWeakBlueprintP
 			}
 			else
 			{
-				FMessageDialog::Open(EAppMsgType::Ok, LOCTEXT("FailedToLoadWidgetBlueprint", "Editor Utility Widget could not be loaded because it derives from an invalid class.\nCheck to make sure the parent class for this blueprint hasn't been removed!"));
+				FMessageDialog::Open(EAppMsgType::Ok, LOCTEXT("FailedToLoadEditorUtilityWidgetBlueprint", "Editor Utility Widget could not be loaded because it derives from an invalid class.\nCheck to make sure the parent class for this blueprint hasn't been removed!"));
 			}
 		}
 	}
diff --git a/Engine/Source/Editor/ContentBrowser/Private/DragDropHandler.cpp b/Engine/Source/Editor/ContentBrowser/Private/DragDropHandler.cpp
--- a/Engine/Source/Editor/ContentBrowser/Private/DragDropHandler.cpp
+++ b/Engine/Source/Editor/ContentBrowser/Private/DragDropHandler.cpp
@@ -132,7 +132,7 @@ void DragDropHandler::HandleDropOnAssetFolder(const TSharedRef<SWidget>& ParentW
 			);
 		MenuBuilder.AddMenuEntry(
 			LOCTEXT("DragDropAdvancedCopy", "Advanced Copy Here"),
-			LOCTEXT("DragDropMoveTooltip", "Copy the dragged items and any specified dependencies to this folder, afterwards fixing up any dependencies on copied files to the new files."),
+			LOCTEXT("DragDropAdvancedCopyTooltip", "Copy the dragged items and any specified dependencies to this folder, afterwards fixing up any dependencies on copied files to the new files."),
 			FSlateIcon(),
 			FUIAction(FExecuteAction::CreateLambda([=]() { AdvancedCopyActionHandler.ExecuteIfBound(FinalAssetList, FinalAssetPaths, TargetPath); }))
 		);
diff --git a/Engine/Source/Editor/DetailCustomizations/Private/FbxImportUIDetails.cpp b/Engine/Source/Editor/DetailCustomizations/Private/FbxImportUIDetails.cpp
--- a/Engine/Source/Editor/DetailCustomizations/Private/FbxImportUIDetails.cpp
+++ b/Engine/Source/Editor/DetailCustomizations/Private/FbxImportUIDetails.cpp
@@ -187,11 +187,11 @@ bool FFbxImportUIDetails::ShowCompareResult()
 				}
 				else if ((SkeletonCompareResult & ImportCompareHelper::ECompareResult::SCR_SkeletonMissingBone) > ImportCompareHelper::ECompareResult::SCR_None)
 				{
-					IconTooltip = (LOCTEXT("SkeletonConflictBadRoot_IconTooltip", "(Warning) Deleted bones: Some bones of the of the current skeletalmesh asset are not use by the incoming fbx."));
+					IconTooltip = (LOCTEXT("SkeletonConflictDeletedBones_IconTooltip", "(Warning) Deleted bones: Some bones of the of the current skeletalmesh asset are not use by the incoming fbx."));
 				}
 				else
 				{
-					IconTooltip = (LOCTEXT("SkeletonConflictBadRoot_IconTooltip", "(Info) Added bones: Some bones in the incoming fbx do not exist in the current skeletalmesh asset."));
+					IconTooltip = (LOCTEXT("SkeletonConflictAddedBones_IconTooltip", "(Info) Added bones: Some bones in the incoming fbx do not exist in the current skeletalmesh asset."));
 				}
 				 
 				const FSlateBrush* Brush = (SkeletonCompareResult & ImportCompareHelper::ECompareResult::SCR_SkeletonBadRoot) > ImportCompareHelper::ECompareResult::SCR_None ? FEditorStyle::GetBrush("Icons.Error")
diff --git a/Engine/Source/Editor/Persona/Private/PersonaMeshDetails.cpp b/Engine/Source/Editor/Persona/Private/PersonaMeshDetails.cpp
--- a/Engine/Source/Editor/Persona/Private/PersonaMeshDetails.cpp
+++ b/Engine/Source/Editor/Persona/Private/PersonaMeshDetails.cpp
@@ -1521,7 +1521,7 @@ FReply FPersonaMeshDetails::RegenerateLOD(int32 LODIndex)
 			else if (bIsReductionActive)
 			{
 				//Ask user a special permission when the base LOD can be reduce 
-				const FText Text(LOCTEXT("Warning_ReductionApplyingToImportedMesh_ReduceBaseLOD", "Are you sure you'd like to apply mesh reduction to the non-generated base LOD?"));
+				const FText Text(LOCTEXT("Warning_ReductionApplyingToImportedMesh_ReduceNonGenBaseLOD", "Are you sure you'd like to apply mesh reduction to the non-generated base LOD?"));
 				EAppReturnType::Type Ret = FMessageDialog::Open(EAppMsgType::YesNo, Text);
 				if (Ret == EAppReturnType::No)
 				{
diff --git a/Engine/Source/Editor/UnrealEd/Private/EditorActorFolders.cpp b/Engine/Source/Editor/UnrealEd/Private/EditorActorFolders.cpp
--- a/Engine/Source/Editor/UnrealEd/Private/EditorActorFolders.cpp
+++ b/Engine/Source/Editor/UnrealEd/Private/EditorActorFolders.cpp
@@ -330,7 +330,7 @@ FName FActorFolders::GetFolderName(UWorld& InWorld, FName ParentPath, FName InLe
 	NumberFormat.SetUseGrouping(false);
 	NumberFormat.SetMinimumIntegralDigits(SuffixLen);
 
-	FText LeafName = FText::Format(LOCTEXT("DefaultFolderNamePattern", "{0}{1}"), FText::FromString(LeafNameRoot), bHasSuffix ? FText::AsNumber(Suffix++, &NumberFormat) : FText::GetEmpty());
+	FText LeafName = FText::Format(LOCTEXT("FolderNamePattern", "{0}{1}"), FText::FromString(LeafNameRoot), bHasSuffix ? FText::AsNumber(Suffix++, &NumberFormat) : FText::GetEmpty());
 
 	FString ParentFolderPath = ParentPath.IsNone() ? TEXT("") : ParentPath.ToString();
 	if (!ParentFolderPath.IsEmpty())
@@ -341,7 +341,7 @@ FName FActorFolders::GetFolderName(UWorld& InWorld, FName ParentPath, FName InLe
 	FName FolderName(*(ParentFolderPath + LeafName.ToString()));
 	while (ExistingFolders.Contains(FolderName))
 	{
-		LeafName = FText::Format(LOCTEXT("DefaultFolderNamePattern", "{0}{1}"), FText::FromString(LeafNameRoot), FText::AsNumber(Suffix++, &NumberFormat));
+		LeafName = FText::Format(LOCTEXT("FolderNamePattern", "{0}{1}"), FText::FromString(LeafNameRoot), FText::AsNumber(Suffix++, &NumberFormat));
 		FolderName = FName(*(ParentFolderPath + LeafName.ToString()));
 		if (Suffix == 0)
 		{
