Fixing stale suggestions staying in the material analyzer list

diff --git a/Engine/Plugins/Editor/MaterialAnalyzer/Source/Private/SMaterialAnalyzer.cpp b/Engine/Plugins/Editor/MaterialAnalyzer/Source/Private/SMaterialAnalyzer.cpp
--- a/Engine/Plugins/Editor/MaterialAnalyzer/Source/Private/SMaterialAnalyzer.cpp
+++ b/Engine/Plugins/Editor/MaterialAnalyzer/Source/Private/SMaterialAnalyzer.cpp
@@ -341,6 +341,7 @@ void SMaterialAnalyzer::Tick(const FGeometry& AllottedGeometry, const double InC
 
 			
 			int32 BackgroundColorCounter = 0;
+			SuggestionDataArray.Empty();
 			for (auto It = Suggestions.CreateConstIterator(); It; ++It)
 			{
 				TSharedPtr<FPermutationSuggestionView> SuggestionHeader = MakeShareable(new FPermutationSuggestionView());
@@ -899,6 +900,7 @@ void FAnalyzeForIdenticalPermutationsAsyncTask::DoWork()
 
 void FAnalyzeForIdenticalPermutationsAsyncTask::GatherSuggestions()
 {
+	Suggestions.Empty();
 	for (TPair<uint32, TArray<FName>>& IdenticalPermutations : MaterialPermutationHashToMaterialObjectPath)
 	{
 		if (IdenticalPermutations.Value.Num() > 1)
