Updating messaging as far as static parameters go
#rnx

diff --git a/Engine/Plugins/Editor/MaterialAnalyzer/Source/Private/SMaterialAnalyzer.cpp b/Engine/Plugins/Editor/MaterialAnalyzer/Source/Private/SMaterialAnalyzer.cpp
--- a/Engine/Plugins/Editor/MaterialAnalyzer/Source/Private/SMaterialAnalyzer.cpp
+++ b/Engine/Plugins/Editor/MaterialAnalyzer/Source/Private/SMaterialAnalyzer.cpp
@@ -922,8 +922,8 @@ void FAnalyzeForIdenticalPermutationsAsyncTask::GatherSuggestions()
 				AllNames.Add(PermutationString);
 			}
 
-			FPermutationSuggestionData NewData = FPermutationSuggestionData(FText::Format(LOCTEXT("IdenticalPermutationSuggestions",
-				"The following {0} materials all have identical permutations."),
+			FPermutationSuggestionData NewData = FPermutationSuggestionData(FText::Format(LOCTEXT("IdenticalStaticPermutationSuggestions",
+				"The following {0} materials all have identical static parameter permutations."),
 				FText::AsNumber(AssetCount)),
 				AllNames);
 
