Fix crash in material analyzer when a StaticComponentMaskParameter was present.

diff --git a/Engine/Plugins/Editor/MaterialAnalyzer/Source/Private/SMaterialAnalyzer.cpp b/Engine/Plugins/Editor/MaterialAnalyzer/Source/Private/SMaterialAnalyzer.cpp
--- a/Engine/Plugins/Editor/MaterialAnalyzer/Source/Private/SMaterialAnalyzer.cpp
+++ b/Engine/Plugins/Editor/MaterialAnalyzer/Source/Private/SMaterialAnalyzer.cpp
@@ -865,11 +865,11 @@ void FAnalyzeMaterialTreeAsyncTask::DoWork()
 		bool bIsOverridden = false;
 		if (CurrentMaterialInstance)
 		{
-			bIsOverridden = CurrentMaterialInstance->GetParameterOverrideValue(EMaterialParameterType::StaticComponentMask, StaticSwitchParameterInfo[ParameterIndex], Meta);
+			bIsOverridden = CurrentMaterialInstance->GetParameterOverrideValue(EMaterialParameterType::StaticComponentMask, StaticMaskParameterInfo[ParameterIndex], Meta);
 		}
 		else if (CurrentMaterial)
 		{
-			bIsOverridden = CurrentMaterial->GetParameterValue(EMaterialParameterType::StaticComponentMask, StaticSwitchParameterInfo[ParameterIndex], Meta);
+			bIsOverridden = CurrentMaterial->GetParameterValue(EMaterialParameterType::StaticComponentMask, StaticMaskParameterInfo[ParameterIndex], Meta);
 		}
 
 		if(!bIsOverridden)
