[Materials] Add parameter filter to Material Analyzer and improve readability of UI elements.

This adds a text field to the UI of the Material Analyzer, which applies a search filter to all displayed material instance property overrides of the actively selected material.
Since the hiearchy of the UI tree view pane is tied to the analyzed material nodes, this simply re-creates the analyzed tree with the applied filter.

#rnx

diff --git a/Engine/Plugins/Editor/MaterialAnalyzer/Source/Private/AnalyzedMaterialNode.h b/Engine/Plugins/Editor/MaterialAnalyzer/Source/Private/AnalyzedMaterialNode.h
--- a/Engine/Plugins/Editor/MaterialAnalyzer/Source/Private/AnalyzedMaterialNode.h
+++ b/Engine/Plugins/Editor/MaterialAnalyzer/Source/Private/AnalyzedMaterialNode.h
@@ -7,11 +7,12 @@
 struct FBasePropertyOverrideNode
 {
 public:
-	FBasePropertyOverrideNode(FName InParameterName, FName InParameterID, float InParameterValue, bool bInOverride) :
+	FBasePropertyOverrideNode(FName InParameterName, FName InParameterID, float InParameterValue, bool bInOverride, const FText& InHighlightText = FText::GetEmpty()) :
 		ParameterName(InParameterName),
 		ParameterID(InParameterID),
 		ParameterValue(InParameterValue),
-		bOverride(bInOverride)
+		bOverride(bInOverride),
+		HighlightText(InHighlightText)
 	{
 
 	}
@@ -19,6 +20,7 @@ public:
 	FName ParameterID;
 	float ParameterValue;
 	bool bOverride;
+	FText HighlightText;
 
 	TArray<TSharedRef<FBasePropertyOverrideNode>>* Children;
 };
@@ -26,31 +28,35 @@ public:
 struct FStaticMaterialLayerParameterNode
 {
 public:
-	FStaticMaterialLayerParameterNode(FName InParameterName, FString InParameterValue, bool bInOverride):
+	FStaticMaterialLayerParameterNode(FName InParameterName, FString InParameterValue, bool bInOverride, const FText& InHighlightText = FText::GetEmpty()):
 		ParameterName(InParameterName),
 		ParameterValue(InParameterValue),
-		bOverride(bInOverride)
+		bOverride(bInOverride),
+		HighlightText(InHighlightText)
 	{
 
 	}
 	FName ParameterName;
 	FString ParameterValue;
 	bool bOverride;
+	FText HighlightText;
 };
 
 struct FStaticSwitchParameterNode
 {
 public:
-	FStaticSwitchParameterNode(FName InParameterName, bool InParameterValue, bool bInOverride) :
+	FStaticSwitchParameterNode(FName InParameterName, bool InParameterValue, bool bInOverride, const FText& InHighlightText = FText::GetEmpty()) :
 		ParameterName(InParameterName),
 		ParameterValue(InParameterValue),
-		bOverride(bInOverride)
+		bOverride(bInOverride),
+		HighlightText(InHighlightText)
 	{
 
 	}
 	FName ParameterName;
 	bool ParameterValue;
 	bool bOverride;
+	FText HighlightText;
 
 	TArray<TSharedRef<FStaticSwitchParameterNode>>* Children;
 };
@@ -58,13 +64,14 @@ public:
 struct FStaticComponentMaskParameterNode
 {
 public:
-	FStaticComponentMaskParameterNode(FName InParameterName, bool InR, bool InG, bool InB, bool InA, bool bInOverride) :
+	FStaticComponentMaskParameterNode(FName InParameterName, bool InR, bool InG, bool InB, bool InA, bool bInOverride, const FText& InHighlightText = FText::GetEmpty()) :
 		ParameterName(InParameterName),
 		R(InR),
 		G(InG),
 		B(InB),
 		A(InA),
-		bOverride(bInOverride)
+		bOverride(bInOverride),
+		HighlightText(InHighlightText)
 	{
 
 	}
@@ -74,6 +81,7 @@ public:
 	bool B;
 	bool A;
 	bool bOverride;
+	FText HighlightText;
 };
 
 typedef TSharedRef<FBasePropertyOverrideNode, ESPMode::ThreadSafe> FBasePropertyOverrideNodeRef;
@@ -113,18 +121,18 @@ public:
 		return &ChildNodes;
 	}
 
-	int ActualNumberOfChildren() const
+	int32 ActualNumberOfChildren() const
 	{
 		return ChildNodes.Num();
 	}
 
-	int TotalNumberOfChildren() const
+	int32 TotalNumberOfChildren() const
 	{
 		int32 TotalChildren = 0;
 
-		for(int i = 0; i < ChildNodes.Num(); ++i)
+		for(const FAnalyzedMaterialNodeRef& ChildNode : ChildNodes)
 		{
-			TotalChildren += ChildNodes[i]->TotalNumberOfChildren();
+			TotalChildren += ChildNode->TotalNumberOfChildren();
 		}
 
 		return TotalChildren + ChildNodes.Num();
@@ -158,6 +166,53 @@ public:
 		return *StaticComponentMaskParameter;
 	}
 
+	bool HasAnyFilteredParameters(const FString& ParameterFilter) const
+	{
+		// Check if this material node has any filtered parameters
+		const bool bHasAnyFilteredBaseProperties = BasePropertyOverrides.FindByPredicate(
+			[&](const FBasePropertyOverrideNodeRef& Entry) -> bool
+			{
+				// Only overridden parameters are displayed, so ignore any inherited parameters for this search
+				return Entry->bOverride && Entry->ParameterName.ToString().Contains(ParameterFilter);
+			}) != nullptr;
+		if (bHasAnyFilteredBaseProperties)
+		{
+			return true;
+		}
+
+		const bool bHasAnyFilteredStaticSwitchParameters = StaticSwitchParameters.FindByPredicate(
+			[&](const FStaticSwitchParameterNodeRef& Entry) -> bool
+			{
+				// Only overridden parameters are displayed, so ignore any inherited parameters for this search
+				return Entry->bOverride && Entry->ParameterName.ToString().Contains(ParameterFilter);
+			}) != nullptr;
+		if (bHasAnyFilteredStaticSwitchParameters)
+		{
+			return true;
+		}
+
+		const bool bHasAnyFilteredStaticComponentMaskParameters = StaticComponentMaskParameters.FindByPredicate(
+			[&](const FStaticComponentMaskParameterNodeRef& Entry) -> bool
+			{
+				// Only overridden parameters are displayed, so ignore any inherited parameters for this search
+				return Entry->bOverride && Entry->ParameterName.ToString().Contains(ParameterFilter);
+			}) != nullptr;
+		if (bHasAnyFilteredStaticComponentMaskParameters)
+		{
+			return true;
+		}
+
+		// Now check child nodes recursively
+		for (const FAnalyzedMaterialNodeRef& Child : ChildNodes)
+		{
+			if (Child->HasAnyFilteredParameters(ParameterFilter))
+			{
+				return true;
+			}
+		}
+		return false;
+	}
+
 	FString Path;
 	FSoftObjectPath ObjectPath;
 	FAnalyzedMaterialNodePtr Parent;
diff --git a/Engine/Plugins/Editor/MaterialAnalyzer/Source/Private/SAnalyzedMaterialNodeWidgetItem.cpp b/Engine/Plugins/Editor/MaterialAnalyzer/Source/Private/SAnalyzedMaterialNodeWidgetItem.cpp
--- a/Engine/Plugins/Editor/MaterialAnalyzer/Source/Private/SAnalyzedMaterialNodeWidgetItem.cpp
+++ b/Engine/Plugins/Editor/MaterialAnalyzer/Source/Private/SAnalyzedMaterialNodeWidgetItem.cpp
@@ -217,6 +217,7 @@ TSharedRef<SWidget> SBasePropertyOverrideWidget::CreateRowWidget(FBasePropertyOv
 		[
 			SNew(STextBlock)
 			.Text(FText::FromName(RowData->ParameterName))
+			.HighlightText(RowData->HighlightText)
 		]
 	+ SHorizontalBox::Slot()
 		.Padding(2, 0, 0, 0)
@@ -241,6 +242,7 @@ TSharedRef<SWidget> SStaticSwitchParameterWidget::CreateRowWidget(FStaticSwitchP
 		[
 			SNew(STextBlock)
 			.Text(FText::FromName(RowData->ParameterName))
+			.HighlightText(RowData->HighlightText)
 		]
 	+ SHorizontalBox::Slot()
 		.Padding(2, 0, 0, 0)
@@ -267,6 +269,7 @@ TSharedRef<SWidget> SStaticComponentMaskParameterWidget::CreateRowWidget(FStatic
 		[
 			SNew(STextBlock)
 			.Text(FText::FromName(RowData->ParameterName))
+			.HighlightText(RowData->HighlightText)
 		]
 		+ SHorizontalBox::Slot()
 		.HAlign(HAlign_Right)
@@ -356,6 +359,7 @@ TSharedRef<SWidget> SStaticMaterialLayerParameterWidget::CreateRowWidget(FStatic
 		[
 			SNew(STextBlock)
 			.Text(FText::FromName(RowData->ParameterName))
+			.HighlightText(RowData->HighlightText)
 		]
 	+ SHorizontalBox::Slot()
 		.Padding(2, 0, 0, 0)
diff --git a/Engine/Plugins/Editor/MaterialAnalyzer/Source/Private/SMaterialAnalyzer.cpp b/Engine/Plugins/Editor/MaterialAnalyzer/Source/Private/SMaterialAnalyzer.cpp
--- a/Engine/Plugins/Editor/MaterialAnalyzer/Source/Private/SMaterialAnalyzer.cpp
+++ b/Engine/Plugins/Editor/MaterialAnalyzer/Source/Private/SMaterialAnalyzer.cpp
@@ -17,6 +17,7 @@
 #include "Widgets/Input/SEditableText.h"
 #include "Widgets/Layout/SScrollBox.h"
 #include "Widgets/Views/STreeView.h"
+#include "Styling/StyleColors.h"
 
 #define LOCTEXT_NAMESPACE "MaterialAnalyzer"
 
@@ -26,7 +27,6 @@ SMaterialAnalyzer::SMaterialAnalyzer()
 	: BuildBaseMaterialTreeTask(nullptr)
 	, AnalyzeTreeTask(nullptr)
 	, AnalyzeForIdenticalPermutationsTask(nullptr)
-	, bRequestedTreeRefresh(false)
 	, bWaitingForAssetRegistryLoad(false)
 {
 	BasePropertyOverrideNames.Empty();
@@ -123,6 +123,14 @@ void SMaterialAnalyzer::Construct(const FArguments& InArgs, const TSharedRef<SDo
 				[
 					AssetPickerWidget
 				]
+				+ SHorizontalBox::Slot()
+				.FillWidth(0.5f)
+				[
+					SNew(SEditableText)
+					.HintText(LOCTEXT("MaterialParametersToFilterHint", "Parameters to Filter"))
+					.ColorAndOpacity(FSlateColor::UseForeground())
+					.OnTextCommitted(this, &SMaterialAnalyzer::OnParameterFilterChanged)
+				]
 				+SHorizontalBox::Slot()
 				.FillWidth(0.5f)
 				[
@@ -238,12 +246,36 @@ void SMaterialAnalyzer::OnAssetAdded(const FAssetData& InAssetData)
 
 void SMaterialAnalyzer::OnAssetSelected(const FAssetData& AssetData)
 {
-	if(AnalyzeTreeTask == nullptr)
+	if (AnalyzeTreeTask == nullptr)
 	{
 		CurrentlySelectedAsset = AssetData;
+		UpdateViewForSelectedAsset();
+	}
+}
+
+void SMaterialAnalyzer::OnParameterFilterChanged(const FText& Filter, const ETextCommit::Type InTextAction)
+{
+	if (AnalyzeTreeTask == nullptr)
+	{
+		ParameterFilter = Filter;
+		bHasParameterFilterChanged = true;
 
-		const FAssetData* ParentAssetData = &AssetData;
-		const FAssetData* NextParentAssetData = FindParentAssetData(&AssetData, AssetDataArray);
+		// Parameter filter is applied in two steps:
+		// 1. Re-build analyzed material tree in order to highlight the filtered parameters (FBuildBasicMaterialTreeAsyncTask)
+		// 2. Apply filter to UI elements by re-analyzing material tree (FAnalyzeMaterialTreeAsyncTask)
+		if (CurrentlySelectedAsset.IsValid())
+		{
+			RecentlyAddedAssetData.Add(CurrentlySelectedAsset);
+		}
+	}
+}
+
+void SMaterialAnalyzer::UpdateViewForSelectedAsset()
+{
+	if (CurrentlySelectedAsset.IsValid())
+	{
+		const FAssetData* ParentAssetData = &CurrentlySelectedAsset;
+		const FAssetData* NextParentAssetData = FindParentAssetData(&CurrentlySelectedAsset, AssetDataArray);
 		// get the topmost parent
 		while (NextParentAssetData != nullptr)
 		{
@@ -266,7 +298,7 @@ void SMaterialAnalyzer::OnAssetSelected(const FAssetData& AssetData)
 
 		SuggestionsBox->ClearChildren();
 
-		AnalyzeTreeTask = new FAsyncTask<FAnalyzeMaterialTreeAsyncTask>(*NewRoot, AssetDataArray);
+		AnalyzeTreeTask = new FAsyncTask<FAnalyzeMaterialTreeAsyncTask>(*NewRoot, AssetDataArray, ParameterFilter);
 
 		StartAsyncWork(FText::Format(LOCTEXT("AnalyzingMaterial", "Analyzing {0}"), FText::FromString(AnalyzeTreeTask->GetTask().CurrentMaterialNode->Path)));
 		AnalyzeTreeTask->StartBackgroundTask();
@@ -295,6 +327,17 @@ void SMaterialAnalyzer::Tick(const FGeometry& AllottedGeometry, const double InC
 			BuildBaseMaterialTreeTask = nullptr;
 			AsyncWorkFinished(FText::Format(FTextFormat(LOCTEXT("DoneWithMaterialInterfaces", "Done with {0} MaterialInterfaces")), GetTotalNumberOfMaterialNodes()));
 
+			// If the parameter filter has changed, we also have to run analysis async task again,
+			// which applies the actual parameter filter to UI elements in the tree view pane.
+			if (bHasParameterFilterChanged)
+			{
+				UpdateViewForSelectedAsset();
+				bHasParameterFilterChanged = false;
+			}
+			else
+			{
+				MaterialTree->RequestTreeRefresh();
+			}
 		}
 
 		if (BuildBaseMaterialTreeTask == nullptr && RecentlyAddedAssetData.Num() > 0)
@@ -429,7 +472,7 @@ TSharedRef< ITableRow > SMaterialAnalyzer::OnGenerateSuggestionRow(TSharedPtr<FP
 							[
 								SNew(SImage)
 								.Image(FAppStyle::Get().GetBrush("Icons.Plus"))
-								.ColorAndOpacity(FSlateColor::UseForeground())
+								.ColorAndOpacity(FSlateColor(EStyleColor::Black))
 							]
 							+SHorizontalBox::Slot()
 							.AutoWidth()
@@ -437,6 +480,7 @@ TSharedRef< ITableRow > SMaterialAnalyzer::OnGenerateSuggestionRow(TSharedPtr<FP
 							[
 								SNew(STextBlock)
 								.Text(LOCTEXT("CreateLocalCollection", "Create Local Collection"))
+								.ColorAndOpacity(FSlateColor(EStyleColor::Black))
 							]
 						]
 					]
@@ -509,6 +553,8 @@ FReply SMaterialAnalyzer::CreateLocalSuggestionCollection(TSharedPtr<FPermutatio
 
 void SMaterialAnalyzer::StartAsyncWork(const FText& WorkText)
 {
+	bIsAsyncWorkInProgress = true;
+
 	if(StatusBox.IsValid())
 	{
 		StatusBox->SetText(WorkText);
@@ -519,8 +565,6 @@ void SMaterialAnalyzer::StartAsyncWork(const FText& WorkText)
 		StatusThrobber->SetAnimate(SThrobber::Horizontal);
 		StatusThrobber->SetVisibility(EVisibility::SelfHitTestInvisible);
 	}
-
-	bAllowMaterialSelection = false;
 }
 
 void SMaterialAnalyzer::AsyncWorkFinished(const FText& CompleteText)
@@ -536,7 +580,7 @@ void SMaterialAnalyzer::AsyncWorkFinished(const FText& CompleteText)
 		StatusThrobber->SetVisibility(EVisibility::Collapsed);
 	}
 
-	bAllowMaterialSelection = true;
+	bIsAsyncWorkInProgress = false;
 }
 
 int32 SMaterialAnalyzer::GetTotalNumberOfMaterialNodes()
@@ -597,7 +641,21 @@ TSharedRef<ITableRow> SMaterialAnalyzer::HandleReflectorTreeGenerateRow(FAnalyze
 
 void SMaterialAnalyzer::HandleReflectorTreeGetChildren(FAnalyzedMaterialNodeRef InMaterialNode, TArray<FAnalyzedMaterialNodeRef>& OutChildren)
 {
-	OutChildren = InMaterialNode->GetChildNodes();
+	if (ParameterFilter.IsEmpty())
+	{
+		OutChildren = InMaterialNode->GetChildNodes();
+	}
+	else
+	{
+		OutChildren.Empty(InMaterialNode->GetChildNodes().Num());
+		for (const FAnalyzedMaterialNodeRef& ChildNode : InMaterialNode->GetChildNodes())
+		{
+			if (ChildNode->HasAnyFilteredParameters(ParameterFilter.ToString()))
+			{
+				OutChildren.Add(ChildNode);
+			}
+		}
+	}
 }
 
 void SMaterialAnalyzer::HandleReflectorTreeRecursiveExpansion(FAnalyzedMaterialNodeRef InTreeNode, bool bIsItemExpanded)
@@ -711,11 +769,22 @@ void FAnalyzeMaterialTreeAsyncTask::DoWork()
 	
 	CurrentMaterialNode->BasePropertyOverrides.Empty(BasePropertyOverrideNames.Num());
 
+	const FString ParameterFilterString = this->ParameterFilter.ToString();
+	auto IsIncludedInParameterFilter = [&ParameterFilterString](const FName& Name) -> bool
+		{
+			return ParameterFilterString.IsEmpty() || Name.ToString().Contains(ParameterFilterString);
+		};
+
 	for (const TPair<FName, FName>& BasePropertyOverrideName : BasePropertyOverrideNames)
 	{
 		float TempValue = 0.0f;
 		bool bIsOverridden = false;
 
+		if (!IsIncludedInParameterFilter(BasePropertyOverrideName.Value))
+		{
+			continue;
+		}
+
 		if (BasePropertyOverrideName.Key.IsEqual(TEXT("bOverride_OpacityMaskClipValue")))
 		{
 			TempValue = CurrentMaterialInterface->GetOpacityMaskClipValue();
@@ -849,12 +918,13 @@ void FAnalyzeMaterialTreeAsyncTask::DoWork()
 				new FBasePropertyOverrideNode(ParentParameter->ParameterName,
 					ParentParameter->ParameterID,
 					ParentParameter->ParameterValue,
-					false)));	
+					false,
+					this->ParameterFilter)));	
 		}
 		else
 		{
 			CurrentMaterialNode->BasePropertyOverrides.Add(FBasePropertyOverrideNodeRef(
-				new FBasePropertyOverrideNode(BasePropertyOverrideName.Value, BasePropertyOverrideName.Key, TempValue, bIsOverridden)));
+				new FBasePropertyOverrideNode(BasePropertyOverrideName.Value, BasePropertyOverrideName.Key, TempValue, bIsOverridden, this->ParameterFilter)));
 		}
 	}
 
@@ -872,11 +942,12 @@ void FAnalyzeMaterialTreeAsyncTask::DoWork()
 		CurrentMaterialNode->MaterialLayerParameters.Add(FStaticMaterialLayerParameterNodeRef(
 			new FStaticMaterialLayerParameterNode(FName(),
 				MaterialLayers.GetStaticPermutationString(),
-				bIsOverridden)));
+				bIsOverridden,
+				this->ParameterFilter)));
 	}
 	
 	CurrentMaterialNode->StaticSwitchParameters.Empty(StaticSwitchParameterInfo.Num());
-	
+
 	for (int ParameterIndex = 0; ParameterIndex < StaticSwitchParameterInfo.Num(); ++ParameterIndex)
 	{
 		FMaterialParameterMetadata Meta;
@@ -890,6 +961,11 @@ void FAnalyzeMaterialTreeAsyncTask::DoWork()
 			bIsOverridden = CurrentMaterial->GetParameterValue(EMaterialParameterType::StaticSwitch, StaticSwitchParameterInfo[ParameterIndex], Meta);
 		}
 
+		if (!IsIncludedInParameterFilter(StaticSwitchParameterInfo[ParameterIndex].Name))
+		{
+			continue;
+		}
+
 		if (!bIsOverridden)
 		{
 			// Check the parent for this variable
@@ -902,12 +978,13 @@ void FAnalyzeMaterialTreeAsyncTask::DoWork()
 			CurrentMaterialNode->StaticSwitchParameters.Add(FStaticSwitchParameterNodeRef(
 				new FStaticSwitchParameterNode(ParentParameter->ParameterName,
 					ParentParameter->ParameterValue,
-					false)));
+					false,
+					this->ParameterFilter)));
 		}
 		else
 		{
 			CurrentMaterialNode->StaticSwitchParameters.Add(FStaticSwitchParameterNodeRef(
-				new FStaticSwitchParameterNode(StaticSwitchParameterInfo[ParameterIndex].Name, Meta.Value.AsStaticSwitch(), true)));
+				new FStaticSwitchParameterNode(StaticSwitchParameterInfo[ParameterIndex].Name, Meta.Value.AsStaticSwitch(), true, this->ParameterFilter)));
 		}
 	}
 	
@@ -915,6 +992,11 @@ void FAnalyzeMaterialTreeAsyncTask::DoWork()
 	
 	for (int ParameterIndex = 0; ParameterIndex < StaticMaskParameterInfo.Num(); ++ParameterIndex)
 	{
+		if (!IsIncludedInParameterFilter(StaticMaskParameterInfo[ParameterIndex].Name))
+		{
+			continue;
+		}
+
 		FMaterialParameterMetadata Meta;
 		bool bIsOverridden = false;
 		if (CurrentMaterialInstance)
@@ -941,7 +1023,8 @@ void FAnalyzeMaterialTreeAsyncTask::DoWork()
 					ParentParameter->G,
 					ParentParameter->B,
 					ParentParameter->A,
-					false)));
+					false,
+					this->ParameterFilter)));
 		}
 		else
 		{
@@ -951,7 +1034,8 @@ void FAnalyzeMaterialTreeAsyncTask::DoWork()
 					Meta.Value.Bool[1],
 					Meta.Value.Bool[2],
 					Meta.Value.Bool[3],
-					true)));
+					true,
+					this->ParameterFilter)));
 		}
 	}
 
diff --git a/Engine/Plugins/Editor/MaterialAnalyzer/Source/Private/SMaterialAnalyzer.h b/Engine/Plugins/Editor/MaterialAnalyzer/Source/Private/SMaterialAnalyzer.h
--- a/Engine/Plugins/Editor/MaterialAnalyzer/Source/Private/SMaterialAnalyzer.h
+++ b/Engine/Plugins/Editor/MaterialAnalyzer/Source/Private/SMaterialAnalyzer.h
@@ -24,12 +24,14 @@ public:
 	
 	TArray<FAssetData> AssetDataToAnalyze;
 
+	FString ParameterFilter;
+
 	/** Initializes the variables needed to load and verify the data */
-	FBuildBasicMaterialTreeAsyncTask(TArray<FAnalyzedMaterialNodeRef>& InMaterialTreeRoot, const TArray<FAssetData>& InAssetDataToAnalyze):
-		MaterialTreeRoot(InMaterialTreeRoot),
-		AssetDataToAnalyze(InAssetDataToAnalyze)
+	FBuildBasicMaterialTreeAsyncTask(TArray<FAnalyzedMaterialNodeRef>& InMaterialTreeRoot, const TArray<FAssetData>& InAssetDataToAnalyze, const FString& InParameterFilter = TEXT(""))
+		: MaterialTreeRoot(InMaterialTreeRoot)
+		, AssetDataToAnalyze(InAssetDataToAnalyze)
+		, ParameterFilter(InParameterFilter)
 	{
-
 	}
 
 	FAnalyzedMaterialNodePtr FindOrMakeBranchNode(FAnalyzedMaterialNodePtr ParentNode, const FAssetData* ChildData);
@@ -66,12 +68,15 @@ public:
 	TArray<FMaterialParameterInfo> StaticMaskParameterInfo;
 	TArray<FGuid> StaticMaskGuids;
 
+	FText ParameterFilter;
 
-	FAnalyzeMaterialTreeAsyncTask(FAnalyzedMaterialNodeRef InMaterialTreeRoot, const TArray<FAssetData>& InAssetDataToAnalyze):
+
+	FAnalyzeMaterialTreeAsyncTask(FAnalyzedMaterialNodeRef InMaterialTreeRoot, const TArray<FAssetData>& InAssetDataToAnalyze, const FText& InParameterFilter = FText::GetEmpty()):
 		MaterialTreeRoot(InMaterialTreeRoot),
 		AssetDataToAnalyze(InAssetDataToAnalyze),
 		CurrentMaterialQueueIndex(0),
-		CurrentMaterialNode(InMaterialTreeRoot)
+		CurrentMaterialNode(InMaterialTreeRoot),
+		ParameterFilter(InParameterFilter)
 	{
 		MaterialQueue.Empty(MaterialTreeRoot->TotalNumberOfChildren());
 		MaterialQueue.Add(MaterialTreeRoot);
@@ -236,13 +241,12 @@ protected:
 	FAsyncTask<FAnalyzeMaterialTreeAsyncTask>* AnalyzeTreeTask;
 	FAsyncTask<FAnalyzeForIdenticalPermutationsAsyncTask>* AnalyzeForIdenticalPermutationsTask;
 
-	bool bRequestedTreeRefresh;
-
 	FAssetData CurrentlySelectedAsset;
 
 	bool bWaitingForAssetRegistryLoad;
 
 	void OnAssetSelected(const FAssetData& AssetData);
+	void OnParameterFilterChanged(const FText& Filter, const ETextCommit::Type InTextAction);
 
 	void BuildBasicMaterialTree();
 
@@ -274,11 +278,16 @@ private:
 	 */
 	void HandleReflectorTreeRecursiveExpansion(FAnalyzedMaterialNodeRef InTreeNode, bool bIsItemExpanded);
 
+	void UpdateViewForSelectedAsset();
+
 	bool IsMaterialSelectionAllowed() const
 	{
-		return bAllowMaterialSelection;
+		// Don't allow to change selected material while async work is operating on currently selected asset.
+		return !bIsAsyncWorkInProgress;
 	}
 
 private:
-	bool bAllowMaterialSelection;
+	bool bIsAsyncWorkInProgress = false;
+	FText ParameterFilter;
+	bool bHasParameterFilterChanged = false;
 };
