[Materials] Add "Export to CSV" functionality and UI element to Material Analyzer.

#rnx

diff --git a/Engine/Plugins/Editor/MaterialAnalyzer/Source/MaterialAnalyzer.Build.cs b/Engine/Plugins/Editor/MaterialAnalyzer/Source/MaterialAnalyzer.Build.cs
--- a/Engine/Plugins/Editor/MaterialAnalyzer/Source/MaterialAnalyzer.Build.cs
+++ b/Engine/Plugins/Editor/MaterialAnalyzer/Source/MaterialAnalyzer.Build.cs
@@ -20,7 +20,9 @@ namespace UnrealBuildTool.Rules
 					"SlateCore",
 					"EditorFramework",
 					"UnrealEd",
-					"PropertyEditor"
+					"PropertyEditor",
+					"DesktopPlatform",
+					"ToolWidgets",
 			});
 
             PrivateDependencyModuleNames.AddRange(new string[] {
diff --git a/Engine/Plugins/Editor/MaterialAnalyzer/Source/Private/SMaterialAnalyzer.cpp b/Engine/Plugins/Editor/MaterialAnalyzer/Source/Private/SMaterialAnalyzer.cpp
--- a/Engine/Plugins/Editor/MaterialAnalyzer/Source/Private/SMaterialAnalyzer.cpp
+++ b/Engine/Plugins/Editor/MaterialAnalyzer/Source/Private/SMaterialAnalyzer.cpp
@@ -6,18 +6,23 @@
 #include "PropertyCustomizationHelpers.h"
 #include "CollectionManagerModule.h"
 #include "Framework/Views/TableViewMetadata.h"
+#include "Framework/Application/SlateApplication.h"
 #include "ICollectionManager.h"
 #include "AssetManagerEditorModule.h"
 #include "Materials/Material.h"
 #include "Materials/MaterialInstance.h"
 #include "Misc/Paths.h"
+#include "Misc/FileHelper.h"
 #include "SAnalyzedMaterialNodeWidgetItem.h"
 #include "Widgets/Images/SThrobber.h"
 #include "Widgets/Input/SButton.h"
 #include "Widgets/Input/SEditableText.h"
 #include "Widgets/Layout/SScrollBox.h"
 #include "Widgets/Views/STreeView.h"
+#include "Filters/SFilterSearchBox.h"
 #include "Styling/StyleColors.h"
+#include "DesktopPlatformModule.h"
+#include "ProfilingDebugging/DiagnosticTable.h"
 
 #define LOCTEXT_NAMESPACE "MaterialAnalyzer"
 
@@ -119,22 +124,45 @@ void SMaterialAnalyzer::Construct(const FArguments& InArgs, const TSharedRef<SDo
 					.Text(LOCTEXT("MaterialToAnalyzeLabel", "Material To Analyze: "))
 				]
 				+ SHorizontalBox::Slot()
-				.FillWidth(0.5f)
+				.FillWidth(0.4f)
 				[
 					AssetPickerWidget
 				]
 				+ SHorizontalBox::Slot()
-				.FillWidth(0.5f)
+				.FillWidth(0.4f)
 				[
-					SNew(SEditableText)
-					.HintText(LOCTEXT("MaterialParametersToFilterHint", "Parameters to Filter"))
-					.ColorAndOpacity(FSlateColor::UseForeground())
+					SNew(SFilterSearchBox)
+					.HintText(LOCTEXT("MaterialParametersToFilterHint", "Parameters to Filter..."))
+					.ToolTipText(LOCTEXT("FilterSearchHint", "Type here to search (pressing enter selects the results)"))
 					.OnTextCommitted(this, &SMaterialAnalyzer::OnParameterFilterChanged)
 				]
 				+SHorizontalBox::Slot()
-				.FillWidth(0.5f)
+				.FillWidth(0.2f)
 				[
-					SNullWidget::NullWidget
+					SNew(SButton)
+					.ButtonStyle(FAppStyle::Get(), "ButtonStyle")
+					.OnClicked(this, &SMaterialAnalyzer::OnExportAnalyzedMaterialToCSV)
+					.ContentPadding(FMargin(2.0f))
+					.Content()
+					[
+						SNew(SHorizontalBox)
+						+SHorizontalBox::Slot()
+						.AutoWidth()
+						.Padding(2.0f)
+						[
+							SNew(SImage)
+							.Image(FAppStyle::Get().GetBrush("Icons.Save"))
+							.ColorAndOpacity(FSlateColor(EStyleColor::Black))
+						]
+						+SHorizontalBox::Slot()
+						.AutoWidth()
+						.Padding(2.0f)
+						[
+							SNew(STextBlock)
+							.Text(LOCTEXT("ExportToCSV", "Export to CSV"))
+							.ColorAndOpacity(FSlateColor(EStyleColor::Black))
+						]
+					]
 				]
 			]
 			+ SVerticalBox::Slot()
@@ -269,6 +297,133 @@ void SMaterialAnalyzer::OnParameterFilterChanged(const FText& Filter, const ETex
 	}
 }
 
+void WriteAnalyzedMaterialNodeToCSVStringInternal(const FAnalyzedMaterialNodeRef& Node, FDiagnosticTableWriterCSV& CSVTable)
+{
+	auto AddCSVCell = [&CSVTable](int32 NumElements, const TFunction<FString(int32 ElementIndex)>& ElementNameCallback) -> void
+		{
+			FString Cell;
+			for (int32 ElementIndex = 0; ElementIndex < NumElements; ++ElementIndex)
+			{
+				Cell += ElementNameCallback(ElementIndex);
+				if (ElementIndex + 1 < NumElements)
+				{
+					Cell += TEXT("\n");
+				}
+			}
+			CSVTable.AddColumn(TEXT("%s"), *Cell);
+		};
+
+	FString Output;
+
+	CSVTable.AddColumn(TEXT("%s"), *Node->AssetData.AssetName.ToString());
+
+	AddCSVCell(
+		Node->BasePropertyOverrides.Num(),
+		[Node](int32 ElementIndex)
+		{
+			return FString::Printf(TEXT("%s ( %f )"), *Node->BasePropertyOverrides[ElementIndex]->ParameterName.ToString(), Node->BasePropertyOverrides[ElementIndex]->ParameterValue);
+		}
+	);
+	AddCSVCell(
+		Node->MaterialLayerParameters.Num(),
+		[Node](int32 ElementIndex)
+		{
+			return FString::Printf(TEXT("%s ( %s )"), *Node->MaterialLayerParameters[ElementIndex]->ParameterName.ToString(), *Node->MaterialLayerParameters[ElementIndex]->ParameterValue);
+		}
+	);
+	AddCSVCell(
+		Node->StaticSwitchParameters.Num(),
+		[Node](int32 ElementIndex)
+		{
+			return FString::Printf(TEXT("%s ( %s )"), *Node->StaticSwitchParameters[ElementIndex]->ParameterName.ToString(), Node->StaticSwitchParameters[ElementIndex]->ParameterValue ? TEXT("True") : TEXT("False"));
+		}
+	);
+	AddCSVCell(
+		Node->StaticComponentMaskParameters.Num(),
+		[Node](int32 ElementIndex)
+		{
+			const FStaticComponentMaskParameterNode& Parameter = *Node->StaticComponentMaskParameters[ElementIndex];
+			return FString::Printf(
+				TEXT("%s ( %s%s%s%s )"),
+				*Parameter.ParameterName.ToString(),
+				Parameter.R ? TEXT("R") : TEXT("_"),
+				Parameter.G ? TEXT("G") : TEXT("_"),
+				Parameter.B ? TEXT("B") : TEXT("_"),
+				Parameter.A ? TEXT("A") : TEXT("_")
+			);
+		}
+	);
+
+	CSVTable.CycleRow();
+
+	for (const FAnalyzedMaterialNodeRef& ChildNode : Node->GetChildNodes())
+	{
+		WriteAnalyzedMaterialNodeToCSVStringInternal(ChildNode, CSVTable);
+	}
+}
+
+static void WriteAnalyzedMaterialNodeToCSVString(const FAnalyzedMaterialNodeRef& Node, FDiagnosticTableWriterCSV& CSVTable)
+{
+	CSVTable.AddColumn(TEXT("MATERIAL"));
+	CSVTable.AddColumn(TEXT("BASE PROPERTY OVERRIDES"));
+	CSVTable.AddColumn(TEXT("LAYER PARAMETERS"));
+	CSVTable.AddColumn(TEXT("STATIC SWITCHES"));
+	CSVTable.AddColumn(TEXT("STATIC COMPONENT MASKS"));
+	CSVTable.CycleRow();
+
+	WriteAnalyzedMaterialNodeToCSVStringInternal(Node, CSVTable);
+}
+
+static bool SaveFileDialog(const FString& Title, const FString& FileTypes, FString& OutFilename, FString& InOutLastFileanme)
+{
+	OutFilename.Empty();
+
+	TArray<FString> OutFilenames;
+	if (IDesktopPlatform* DesktopPlatform = FDesktopPlatformModule::Get())
+	{
+		const bool bFileChosen = DesktopPlatform->SaveFileDialog(
+			FSlateApplication::Get().FindBestParentWindowHandleForDialogs(nullptr),
+			Title,
+			FPaths::GetPath(InOutLastFileanme),
+			FPaths::GetPathLeaf(InOutLastFileanme),
+			FileTypes,
+			EFileDialogFlags::None,
+			OutFilenames
+		);
+		if (bFileChosen && OutFilenames.Num() > 0)
+		{
+			// User successfully chose a file; remember the path for the next time the dialog opens.
+			OutFilename = OutFilenames[0];
+			InOutLastFileanme = OutFilenames[0];
+			return true;
+		}
+	}
+
+	return false;
+}
+
+FReply SMaterialAnalyzer::OnExportAnalyzedMaterialToCSV()
+{
+	if (!MaterialTreeRoot.IsEmpty())
+	{
+		FString ExportFilename;
+		static FString LastUsedFilename = FPaths::Combine(FPaths::GetProjectFilePath(), TEXT("Saved"), TEXT("Logs"), TEXT("MaterialProperties.csv"));
+		if (SaveFileDialog(NSLOCTEXT("UnrealEd", "Export", "Export").ToString(), TEXT("Comma Separated Value (CSV) Files|*.csv"), ExportFilename, LastUsedFilename))
+		{
+			if (TUniquePtr<FArchive> CSVTableFile = TUniquePtr<FArchive>{ IFileManager::Get().CreateFileWriter(*ExportFilename) })
+			{
+				FDiagnosticTableWriterCSV CSVTable{ CSVTableFile.Get() };
+				WriteAnalyzedMaterialNodeToCSVString(MaterialTreeRoot[0], CSVTable);
+			}
+			else
+			{
+				FPlatformMisc::MessageBoxExt(EAppMsgType::Ok, TEXT("Export operation failed!"), *NSLOCTEXT("UnrealEd", "Error", "Error").ToString());
+			}
+		}
+	}
+	return FReply::Handled();
+}
+
 void SMaterialAnalyzer::UpdateViewForSelectedAsset()
 {
 	if (CurrentlySelectedAsset.IsValid())
diff --git a/Engine/Plugins/Editor/MaterialAnalyzer/Source/Private/SMaterialAnalyzer.h b/Engine/Plugins/Editor/MaterialAnalyzer/Source/Private/SMaterialAnalyzer.h
--- a/Engine/Plugins/Editor/MaterialAnalyzer/Source/Private/SMaterialAnalyzer.h
+++ b/Engine/Plugins/Editor/MaterialAnalyzer/Source/Private/SMaterialAnalyzer.h
@@ -248,6 +248,8 @@ protected:
 	void OnAssetSelected(const FAssetData& AssetData);
 	void OnParameterFilterChanged(const FText& Filter, const ETextCommit::Type InTextAction);
 
+	FReply OnExportAnalyzedMaterialToCSV();
+
 	void BuildBasicMaterialTree();
 
 	int32 GetTotalNumberOfMaterialNodes();
