Fix NSLOCTEXT validation warning in Material Analyzer.

This is a follow-up fix for CL 33424548.

#rnx

diff --git a/Engine/Plugins/Editor/MaterialAnalyzer/Source/Private/SMaterialAnalyzer.cpp b/Engine/Plugins/Editor/MaterialAnalyzer/Source/Private/SMaterialAnalyzer.cpp
--- a/Engine/Plugins/Editor/MaterialAnalyzer/Source/Private/SMaterialAnalyzer.cpp
+++ b/Engine/Plugins/Editor/MaterialAnalyzer/Source/Private/SMaterialAnalyzer.cpp
@@ -417,7 +417,7 @@ FReply SMaterialAnalyzer::OnExportAnalyzedMaterialToCSV()
 			}
 			else
 			{
-				FPlatformMisc::MessageBoxExt(EAppMsgType::Ok, TEXT("Export operation failed!"), *NSLOCTEXT("UnrealEd", "Error", "Error").ToString());
+				FPlatformMisc::MessageBoxExt(EAppMsgType::Ok, TEXT("Export operation failed!"), *NSLOCTEXT("UnrealEd", "Error", "!").ToString());
 			}
 		}
 	}
