[Materials] Remove unused log category in MaterialAnalyzer.

#rnx

diff --git a/Engine/Plugins/Editor/MaterialAnalyzer/Source/MaterialAnalyzerModule.cpp b/Engine/Plugins/Editor/MaterialAnalyzer/Source/MaterialAnalyzerModule.cpp
--- a/Engine/Plugins/Editor/MaterialAnalyzer/Source/MaterialAnalyzerModule.cpp
+++ b/Engine/Plugins/Editor/MaterialAnalyzer/Source/MaterialAnalyzerModule.cpp
@@ -8,8 +8,6 @@
 
 #define LOCTEXT_NAMESPACE "MaterialAnalyzer"
 
-DEFINE_LOG_CATEGORY(MaterialAnalyzer);
-
 static const FName MaterialAnalyzerName("MaterialAnalyzer");
 
 class FMaterialAnalyzerModule : public IModuleInterface
@@ -28,8 +26,6 @@ public:
 	}
 
 	TSharedRef<SDockTab> SpawnMaterialAnalyzerTab(const FSpawnTabArgs& SpawnTabArgs);
-
-protected:
 };
 
 void FMaterialAnalyzerModule::StartupModule()
diff --git a/Engine/Plugins/Editor/MaterialAnalyzer/Source/MaterialAnalyzerModule.h b/Engine/Plugins/Editor/MaterialAnalyzer/Source/MaterialAnalyzerModule.h
--- a/Engine/Plugins/Editor/MaterialAnalyzer/Source/MaterialAnalyzerModule.h
+++ b/Engine/Plugins/Editor/MaterialAnalyzer/Source/MaterialAnalyzerModule.h
@@ -4,8 +4,6 @@
 
 #include "Logging/LogMacros.h"
 
-DECLARE_LOG_CATEGORY_EXTERN(MaterialAnalyzer, Log, All);
-
 #if UE_ENABLE_INCLUDE_ORDER_DEPRECATED_IN_5_2
 #include "CoreMinimal.h"
 #endif
